Created a plugin for working with light weight instances in the editor. This plugin adds menu support for converting selected actors to light weight instances.

diff --git a/Engine/Plugins/LightWeightInstancesEditor/LightWeightInstancesEditor.uplugin b/Engine/Plugins/LightWeightInstancesEditor/LightWeightInstancesEditor.uplugin
new file mode 100644
--- /dev/null
+++ b/Engine/Plugins/LightWeightInstancesEditor/LightWeightInstancesEditor.uplugin
@@ -0,0 +1,24 @@
+{
+	"FileVersion": 3,
+	"Version": 1,
+	"VersionName": "1.0",
+	"FriendlyName": "LightWeightInstancesEditor",
+	"Description": "Light Weight Instances provide the flexibility and interaction of actors while having performance similar to instanced meshes.",
+	"Category": "Other",
+	"CreatedBy": "Epic Games, Inc.",
+	"CreatedByURL": "",
+	"DocsURL": "",
+	"MarketplaceURL": "",
+	"SupportURL": "",
+	"CanContainContent": true,
+	"IsBetaVersion": true,
+	"IsExperimentalVersion": false,
+	"Installed": false,
+	"Modules": [
+		{
+			"Name": "LightWeightInstancesEditor",
+			"Type": "Runtime",
+			"LoadingPhase": "Default"
+		}
+	]
+}
\ No newline at end of file
diff --git a/Engine/Plugins/LightWeightInstancesEditor/Source/LightWeightInstancesEditor/LightWeightInstancesEditor.Build.cs b/Engine/Plugins/LightWeightInstancesEditor/Source/LightWeightInstancesEditor/LightWeightInstancesEditor.Build.cs
new file mode 100644
--- /dev/null
+++ b/Engine/Plugins/LightWeightInstancesEditor/Source/LightWeightInstancesEditor/LightWeightInstancesEditor.Build.cs
@@ -0,0 +1,53 @@
+// Copyright Epic Games, Inc. All Rights Reserved.
+
+using UnrealBuildTool;
+
+public class LightWeightInstancesEditor : ModuleRules
+{
+	public LightWeightInstancesEditor(ReadOnlyTargetRules Target) : base(Target)
+	{
+		PCHUsage = ModuleRules.PCHUsageMode.UseExplicitOrSharedPCHs;
+		
+		PublicIncludePaths.AddRange(
+			new string[] {
+				// ... add public include paths required here ...
+			}
+			);
+				
+		
+		PrivateIncludePaths.AddRange(
+			new string[] {
+				// ... add other private include paths required here ...
+			}
+			);
+			
+		
+		PublicDependencyModuleNames.AddRange(
+			new string[]
+			{
+				"Core",
+				// ... add other public dependencies that you statically link with here ...
+			}
+			);
+			
+		
+		PrivateDependencyModuleNames.AddRange(
+			new string[]
+			{
+				"CoreUObject",
+				"Engine",
+				"Slate",
+				"SlateCore",
+				// ... add private dependencies that you statically link with here ...	
+			}
+			);
+		
+		
+		DynamicallyLoadedModuleNames.AddRange(
+			new string[]
+			{
+				// ... add any modules that your module loads dynamically here ...
+			}
+			);
+	}
+}
diff --git a/Engine/Plugins/LightWeightInstancesEditor/Source/LightWeightInstancesEditor/Private/LightWeightInstancesEditor.cpp b/Engine/Plugins/LightWeightInstancesEditor/Source/LightWeightInstancesEditor/Private/LightWeightInstancesEditor.cpp
new file mode 100644
--- /dev/null
+++ b/Engine/Plugins/LightWeightInstancesEditor/Source/LightWeightInstancesEditor/Private/LightWeightInstancesEditor.cpp
@@ -0,0 +1,107 @@
+// Copyright Epic Games, Inc. All Rights Reserved.
+
+#include "LightWeightInstancesEditor.h"
+#include "LevelEditor.h"
+#include "ToolMenus.h"
+#include "GameFramework/LightWeightInstanceManager.h"
+#include "GameFramework/LightWeightInstanceSubsystem.h"
+
+#define LOCTEXT_NAMESPACE "FLightWeightInstancesEditorModule"
+
+DEFINE_LOG_CATEGORY_STATIC(LogLWIEditor, Log, All);
+
+typedef FLevelEditorModule::FLevelViewportMenuExtender_SelectedActors DelegateType;
+
+void FLightWeightInstancesEditorModule::StartupModule()
+{
+	// hook up level editor extension for conversion between light weight instances and actors
+	AddLevelViewportMenuExtender();
+}
+
+void FLightWeightInstancesEditorModule::ShutdownModule()
+{
+	// Cleanup menu extenstions
+	RemoveLevelViewportMenuExtender();
+}
+
+void FLightWeightInstancesEditorModule::AddLevelViewportMenuExtender()
+{
+	FLevelEditorModule& LevelEditorModule = FModuleManager::Get().LoadModuleChecked<FLevelEditorModule>("LevelEditor");
+	auto& MenuExtenders = LevelEditorModule.GetAllLevelViewportContextMenuExtenders();
+
+	MenuExtenders.Add(DelegateType::CreateRaw(this, &FLightWeightInstancesEditorModule::CreateLevelViewportContextMenuExtender));
+	LevelViewportExtenderHandle = MenuExtenders.Last().GetHandle();
+}
+
+void FLightWeightInstancesEditorModule::RemoveLevelViewportMenuExtender()
+{
+	if (LevelViewportExtenderHandle.IsValid())
+	{
+		FLevelEditorModule* LevelEditorModule = FModuleManager::Get().GetModulePtr<FLevelEditorModule>("LevelEditor");
+		if (LevelEditorModule)
+		{
+			LevelEditorModule->GetAllLevelViewportContextMenuExtenders().RemoveAll([=](const DelegateType& In) { return In.GetHandle() == LevelViewportExtenderHandle; });
+		}
+	}
+}
+
+
+TSharedRef<FExtender> FLightWeightInstancesEditorModule::CreateLevelViewportContextMenuExtender(const TSharedRef<FUICommandList> CommandList, const TArray<AActor*> InActors)
+{
+	TSharedRef<FExtender> Extender = MakeShareable(new FExtender);
+
+	// We only support conversion if all of the actors are the same type
+	for (AActor* Actor : InActors)
+	{
+		if (Actor->GetClass() != InActors[0]->GetClass())
+		{
+			UE_LOG(LogLWIEditor, Warning, TEXT("Unable to convert actors of multiple types to light weight instances"));
+			return Extender;
+		}
+	}
+
+	if (InActors.Num() > 0)
+	{
+		FText ActorName = InActors.Num() == 1 ? FText::Format(LOCTEXT("ActorNameSingular", "\"{0}\""), FText::FromString(InActors[0]->GetActorLabel())) : LOCTEXT("ActorNamePlural", "Actors");
+
+		FLevelEditorModule& LevelEditor = FModuleManager::GetModuleChecked<FLevelEditorModule>(TEXT("LevelEditor"));
+		TSharedRef<FUICommandList> LevelEditorCommandBindings = LevelEditor.GetGlobalLevelEditorActions();
+
+		// We can only convert to an LWI if we didn't select an LWI
+		if (!InActors[0]->GetClass()->IsChildOf(ALightWeightInstanceManager::StaticClass()))
+		{
+			Extender->AddMenuExtension("ActorControl", EExtensionHook::After, LevelEditorCommandBindings, FMenuExtensionDelegate::CreateLambda(
+				[this, ActorName, InActors](FMenuBuilder& MenuBuilder) {
+
+					MenuBuilder.AddMenuEntry(
+						FText::Format(LOCTEXT("ConvertSelectedActorsToLWIsText", "Convert {0} To Light Weight Instances"), ActorName),
+						LOCTEXT("ConvertSelectedActorsToLWIsTooltip", "Convert the selected actors to light weight instances."),
+						FSlateIcon(),
+						FUIAction(FExecuteAction::CreateRaw(this, &FLightWeightInstancesEditorModule::ConvertActorsToLWIsUIAction, InActors))
+					);
+				})
+			);
+		}
+	}
+
+	return Extender;
+}
+
+void FLightWeightInstancesEditorModule::ConvertActorsToLWIsUIAction(const TArray<AActor*> InActors) const
+{
+	if (InActors.IsEmpty() || InActors[0] == nullptr)
+	{
+		UE_LOG(LogLWIEditor, Log, TEXT("Unable to convert unspecified actors to light weight instances"));
+		return;
+	}
+	ALightWeightInstanceManager* Manager = FLightWeightInstanceSubsystem::Get().FindOrAddLightWeightInstanceManager(InActors[0]->GetClass(), InActors[0]->GetLevel());
+	check(Manager);
+	for (AActor* Actor : InActors)
+	{
+		Manager->ConvertActorToLightWeightInstance(Actor);
+	}
+}
+
+#undef LOCTEXT_NAMESPACE
+	
+IMPLEMENT_MODULE(FLightWeightInstancesEditorModule, LightWeightInstancesEditor)
\ No newline at end of file
diff --git a/Engine/Plugins/LightWeightInstancesEditor/Source/LightWeightInstancesEditor/Public/LightWeightInstancesEditor.h b/Engine/Plugins/LightWeightInstancesEditor/Source/LightWeightInstancesEditor/Public/LightWeightInstancesEditor.h
new file mode 100644
--- /dev/null
+++ b/Engine/Plugins/LightWeightInstancesEditor/Source/LightWeightInstancesEditor/Public/LightWeightInstancesEditor.h
@@ -0,0 +1,30 @@
+// Copyright Epic Games, Inc. All Rights Reserved.
+
+#pragma once
+
+#include "CoreMinimal.h"
+#include "Modules/ModuleManager.h"
+
+class FLightWeightInstancesEditorModule : public IModuleInterface
+{
+public:
+
+	/** IModuleInterface implementation */
+	virtual void StartupModule() override;
+	virtual void ShutdownModule() override;
+
+protected:
+	// Adds new menu options related to LWIs
+	void AddLevelViewportMenuExtender();
+
+	// Cleanup menu options related to LWIs
+	void RemoveLevelViewportMenuExtender();
+
+	TSharedRef<FExtender> CreateLevelViewportContextMenuExtender(const TSharedRef<FUICommandList> CommandList, const TArray<AActor*> InActors);
+
+	// Converts InActors to light weight instances. InActors must all be the same type or conversion won't occur
+	void ConvertActorsToLWIsUIAction(const TArray<AActor*> InActors) const;
+
+	// Delegates
+	FDelegateHandle LevelViewportExtenderHandle;
+};
