Add #ifdefs to fix broken build.

diff --git a/Engine/Plugins/LightWeightInstancesEditor/Source/LightWeightInstancesEditor/Private/LightWeightInstancesEditor.cpp b/Engine/Plugins/LightWeightInstancesEditor/Source/LightWeightInstancesEditor/Private/LightWeightInstancesEditor.cpp
--- a/Engine/Plugins/LightWeightInstancesEditor/Source/LightWeightInstancesEditor/Private/LightWeightInstancesEditor.cpp
+++ b/Engine/Plugins/LightWeightInstancesEditor/Source/LightWeightInstancesEditor/Private/LightWeightInstancesEditor.cpp
@@ -1,8 +1,10 @@
 // Copyright Epic Games, Inc. All Rights Reserved.
 
 #include "LightWeightInstancesEditor.h"
+#if WITH_EDITOR
 #include "LevelEditor.h"
 #include "ToolMenus.h"
+#endif
 #include "GameFramework/LightWeightInstanceManager.h"
 #include "GameFramework/LightWeightInstanceSubsystem.h"
 
@@ -10,7 +12,9 @@
 
 DEFINE_LOG_CATEGORY_STATIC(LogLWIEditor, Log, All);
 
+#if WITH_EDITOR
 typedef FLevelEditorModule::FLevelViewportMenuExtender_SelectedActors DelegateType;
+#endif
 
 void FLightWeightInstancesEditorModule::StartupModule()
 {
@@ -26,15 +30,18 @@ void FLightWeightInstancesEditorModule::ShutdownModule()
 
 void FLightWeightInstancesEditorModule::AddLevelViewportMenuExtender()
 {
+#if WITH_EDITOR
 	FLevelEditorModule& LevelEditorModule = FModuleManager::Get().LoadModuleChecked<FLevelEditorModule>("LevelEditor");
 	auto& MenuExtenders = LevelEditorModule.GetAllLevelViewportContextMenuExtenders();
 
 	MenuExtenders.Add(DelegateType::CreateRaw(this, &FLightWeightInstancesEditorModule::CreateLevelViewportContextMenuExtender));
 	LevelViewportExtenderHandle = MenuExtenders.Last().GetHandle();
+#endif
 }
 
 void FLightWeightInstancesEditorModule::RemoveLevelViewportMenuExtender()
 {
+#if WITH_EDITOR
 	if (LevelViewportExtenderHandle.IsValid())
 	{
 		FLevelEditorModule* LevelEditorModule = FModuleManager::Get().GetModulePtr<FLevelEditorModule>("LevelEditor");
@@ -43,13 +50,14 @@ void FLightWeightInstancesEditorModule::RemoveLevelViewportMenuExtender()
 			LevelEditorModule->GetAllLevelViewportContextMenuExtenders().RemoveAll([=](const DelegateType& In) { return In.GetHandle() == LevelViewportExtenderHandle; });
 		}
 	}
+#endif
 }
 
 
 TSharedRef<FExtender> FLightWeightInstancesEditorModule::CreateLevelViewportContextMenuExtender(const TSharedRef<FUICommandList> CommandList, const TArray<AActor*> InActors)
 {
 	TSharedRef<FExtender> Extender = MakeShareable(new FExtender);
-
+#if WITH_EDITOR
 	// We only support conversion if all of the actors are the same type
 	for (AActor* Actor : InActors)
 	{
@@ -83,7 +91,7 @@ TSharedRef<FExtender> FLightWeightInstancesEditorModule::CreateLevelViewportCont
 			);
 		}
 	}
-
+#endif
 	return Extender;
 }
 
