Fixed LightWeightInstancesEditorModule to work in -game.

diff --git a/Engine/Plugins/LightWeightInstancesEditor/Source/LightWeightInstancesEditor/Private/LightWeightInstancesEditor.cpp b/Engine/Plugins/LightWeightInstancesEditor/Source/LightWeightInstancesEditor/Private/LightWeightInstancesEditor.cpp
--- a/Engine/Plugins/LightWeightInstancesEditor/Source/LightWeightInstancesEditor/Private/LightWeightInstancesEditor.cpp
+++ b/Engine/Plugins/LightWeightInstancesEditor/Source/LightWeightInstancesEditor/Private/LightWeightInstancesEditor.cpp
@@ -4,6 +4,9 @@
 #if WITH_EDITOR
 #include "LevelEditor.h"
 #include "ToolMenus.h"
+#include "Framework/MultiBox/MultiBoxExtender.h"
+#include "Framework/Commands/UICommandList.h"
+#include "GameFramework/Actor.h"
 #endif
 #include "GameFramework/LightWeightInstanceManager.h"
 #include "GameFramework/LightWeightInstanceSubsystem.h"
@@ -31,11 +34,14 @@ void FLightWeightInstancesEditorModule::ShutdownModule()
 void FLightWeightInstancesEditorModule::AddLevelViewportMenuExtender()
 {
 #if WITH_EDITOR
-	FLevelEditorModule& LevelEditorModule = FModuleManager::Get().LoadModuleChecked<FLevelEditorModule>("LevelEditor");
-	auto& MenuExtenders = LevelEditorModule.GetAllLevelViewportContextMenuExtenders();
+	if (!IsRunningGame())
+	{
+		FLevelEditorModule& LevelEditorModule = FModuleManager::Get().LoadModuleChecked<FLevelEditorModule>("LevelEditor");
+		auto& MenuExtenders = LevelEditorModule.GetAllLevelViewportContextMenuExtenders();
 
-	MenuExtenders.Add(DelegateType::CreateRaw(this, &FLightWeightInstancesEditorModule::CreateLevelViewportContextMenuExtender));
-	LevelViewportExtenderHandle = MenuExtenders.Last().GetHandle();
+		MenuExtenders.Add(DelegateType::CreateRaw(this, &FLightWeightInstancesEditorModule::CreateLevelViewportContextMenuExtender));
+		LevelViewportExtenderHandle = MenuExtenders.Last().GetHandle();
+	}
 #endif
 }
 
@@ -53,11 +59,10 @@ void FLightWeightInstancesEditorModule::RemoveLevelViewportMenuExtender()
 #endif
 }
 
-
+#if WITH_EDITOR
 TSharedRef<FExtender> FLightWeightInstancesEditorModule::CreateLevelViewportContextMenuExtender(const TSharedRef<FUICommandList> CommandList, const TArray<AActor*> InActors)
 {
 	TSharedRef<FExtender> Extender = MakeShareable(new FExtender);
-#if WITH_EDITOR
 	// We only support conversion if all of the actors are the same type
 	for (AActor* Actor : InActors)
 	{
@@ -91,7 +96,6 @@ TSharedRef<FExtender> FLightWeightInstancesEditorModule::CreateLevelViewportCont
 			);
 		}
 	}
-#endif
 	return Extender;
 }
 
@@ -110,6 +114,8 @@ void FLightWeightInstancesEditorModule::ConvertActorsToLWIsUIAction(const TArray
 	}
 }
 
+#endif
+
 #undef LOCTEXT_NAMESPACE
 	
 IMPLEMENT_MODULE(FLightWeightInstancesEditorModule, LightWeightInstancesEditor)
\ No newline at end of file
diff --git a/Engine/Plugins/LightWeightInstancesEditor/Source/LightWeightInstancesEditor/Public/LightWeightInstancesEditor.h b/Engine/Plugins/LightWeightInstancesEditor/Source/LightWeightInstancesEditor/Public/LightWeightInstancesEditor.h
--- a/Engine/Plugins/LightWeightInstancesEditor/Source/LightWeightInstancesEditor/Public/LightWeightInstancesEditor.h
+++ b/Engine/Plugins/LightWeightInstancesEditor/Source/LightWeightInstancesEditor/Public/LightWeightInstancesEditor.h
@@ -5,6 +5,12 @@
 #include "CoreMinimal.h"
 #include "Modules/ModuleManager.h"
 
+#if WITH_EDITOR
+class AActor;
+class FExtender;
+class FUICommandList;
+#endif
+
 class FLightWeightInstancesEditorModule : public IModuleInterface
 {
 public:
@@ -20,6 +26,7 @@ protected:
 	// Cleanup menu options related to LWIs
 	void RemoveLevelViewportMenuExtender();
 
+#if WITH_EDITOR
 	TSharedRef<FExtender> CreateLevelViewportContextMenuExtender(const TSharedRef<FUICommandList> CommandList, const TArray<AActor*> InActors);
 
 	// Converts InActors to light weight instances. InActors must all be the same type or conversion won't occur
@@ -27,4 +34,5 @@ protected:
 
 	// Delegates
 	FDelegateHandle LevelViewportExtenderHandle;
+#endif
 };
