Fix missing LWI editor file from CL 24263247.

diff --git a/Engine/Plugins/LightWeightInstancesEditor/Source/LightWeightInstancesEditor/Private/LightWeightInstancesEditor.cpp b/Engine/Plugins/LightWeightInstancesEditor/Source/LightWeightInstancesEditor/Private/LightWeightInstancesEditor.cpp
--- a/Engine/Plugins/LightWeightInstancesEditor/Source/LightWeightInstancesEditor/Private/LightWeightInstancesEditor.cpp
+++ b/Engine/Plugins/LightWeightInstancesEditor/Source/LightWeightInstancesEditor/Private/LightWeightInstancesEditor.cpp
@@ -113,7 +113,7 @@ void FLightWeightInstancesEditorModule::ConvertActorsToLWIsUIAction(const TArray
 		TArray<const UDataLayerInstance*> DataLayerInstances = Actor->GetDataLayerInstances();
 		const UDataLayerInstance* DataLayerInstance = DataLayerInstances.Num() > 0 ? DataLayerInstances[0] : nullptr;
 
-		ALightWeightInstanceManager* Manager = FLightWeightInstanceSubsystem::Get().FindOrAddLightWeightInstanceManager(Actor->GetClass(), DataLayerInstance, Actor->GetWorld(), Actor->GetActorLocation());
+		ALightWeightInstanceManager* Manager = FLightWeightInstanceSubsystem::Get().FindOrAddLightWeightInstanceManager(*Actor->GetClass(), *Actor->GetWorld(), Actor->GetActorLocation(), DataLayerInstance);
 		check(Manager);
 		UDataLayerEditorSubsystem::Get()->OnActorDataLayersChanged().Broadcast(Manager);
 
