Avoid nullptr dereference during "Convert <actors> to light weight instances" if data layers aren't enabled.
Initialize HISM random seed during LightWeightInstanceStaticMeshManager setup.

diff --git a/Engine/Plugins/LightWeightInstancesEditor/Source/LightWeightInstancesEditor/Private/LightWeightInstancesEditor.cpp b/Engine/Plugins/LightWeightInstancesEditor/Source/LightWeightInstancesEditor/Private/LightWeightInstancesEditor.cpp
--- a/Engine/Plugins/LightWeightInstancesEditor/Source/LightWeightInstancesEditor/Private/LightWeightInstancesEditor.cpp
+++ b/Engine/Plugins/LightWeightInstancesEditor/Source/LightWeightInstancesEditor/Private/LightWeightInstancesEditor.cpp
@@ -115,7 +115,11 @@ void FLightWeightInstancesEditorModule::ConvertActorsToLWIsUIAction(const TArray
 
 		ALightWeightInstanceManager* Manager = FLightWeightInstanceSubsystem::Get().FindOrAddLightWeightInstanceManager(*Actor->GetClass(), *Actor->GetWorld(), Actor->GetActorLocation(), DataLayerInstance);
 		check(Manager);
-		UDataLayerEditorSubsystem::Get()->OnActorDataLayersChanged().Broadcast(Manager);
+		
+		if(UDataLayerEditorSubsystem* DataLayerEditor = UDataLayerEditorSubsystem::Get())
+		{
+			DataLayerEditor->OnActorDataLayersChanged().Broadcast(Manager);
+		}
 
 		Manager->ConvertActorToLightWeightInstance(Actor);
 	}
diff --git a/Engine/Source/Runtime/Engine/Private/LightWeightInstanceStaticMeshManager.cpp b/Engine/Source/Runtime/Engine/Private/LightWeightInstanceStaticMeshManager.cpp
--- a/Engine/Source/Runtime/Engine/Private/LightWeightInstanceStaticMeshManager.cpp
+++ b/Engine/Source/Runtime/Engine/Private/LightWeightInstanceStaticMeshManager.cpp
@@ -202,7 +202,12 @@ void ALightWeightInstanceStaticMeshManager::SetInstancedStaticMeshParams()
 	InstancedStaticMeshComponent->Mobility = EComponentMobility::Movable;
 	// Allows per-instance selection in the editor
 	InstancedStaticMeshComponent->bHasPerInstanceHitProxies = true;
-}
+
+	while(InstancedStaticMeshComponent->InstancingRandomSeed == 0)
+	{
+		InstancedStaticMeshComponent->InstancingRandomSeed = FMath::Rand();
+	}
+} 
 
 void ALightWeightInstanceStaticMeshManager::SetStaticMeshFromActor(AActor* InActor)
 {
